Dynamotor Synth UPlugin
- Stub for new Dynamotor Vehicle Engine Synthesizer

diff --git a/Engine/Plugins/DynamotorSynth/DynamotorSynth.uplugin b/Engine/Plugins/DynamotorSynth/DynamotorSynth.uplugin
new file mode 100644
--- /dev/null
+++ b/Engine/Plugins/DynamotorSynth/DynamotorSynth.uplugin
@@ -0,0 +1,24 @@
+{
+	"FileVersion": 3,
+	"Version": 1,
+	"VersionName": "0.1",
+	"FriendlyName": "DynamotorSynth",
+	"Description": "A granular synth made for designing dynamic vehicle and motor sounds.",
+	"Category": "Audio",
+	"CreatedBy": "Epic Games, Inc.",
+	"CreatedByURL": "http://epicgames.com",
+	"DocsURL": "",
+	"MarketplaceURL": "",
+	"SupportURL": "",
+	"CanContainContent": true,
+	"IsBetaVersion": true,
+	"IsExperimentalVersion": false,
+	"Installed": false,
+	"Modules": [
+		{
+			"Name": "DynamotorSynth",
+			"Type": "Runtime",
+			"LoadingPhase": "Default"
+		}
+	]
+}
\ No newline at end of file
diff --git a/Engine/Plugins/DynamotorSynth/Source/DynamotorSynth/DynamotorSynth.Build.cs b/Engine/Plugins/DynamotorSynth/Source/DynamotorSynth/DynamotorSynth.Build.cs
new file mode 100644
--- /dev/null
+++ b/Engine/Plugins/DynamotorSynth/Source/DynamotorSynth/DynamotorSynth.Build.cs
@@ -0,0 +1,54 @@
+// Copyright Epic Games, Inc. All Rights Reserved.
+
+using UnrealBuildTool;
+
+public class DynamotorSynth : ModuleRules
+{
+	public DynamotorSynth(ReadOnlyTargetRules Target) : base(Target)
+	{
+		PCHUsage = ModuleRules.PCHUsageMode.UseExplicitOrSharedPCHs;
+		
+		PublicIncludePaths.AddRange(
+			new string[] {
+				// ... add public include paths required here ...
+			}
+			);
+				
+		
+		PrivateIncludePaths.AddRange(
+			new string[] {
+				// ... add other private include paths required here ...
+			}
+			);
+			
+		
+		PublicDependencyModuleNames.AddRange(
+			new string[]
+			{
+				"Core",
+				// ... add other public dependencies that you statically link with here ...
+			}
+			);
+			
+		
+		PrivateDependencyModuleNames.AddRange(
+			new string[]
+			{
+				"CoreUObject",
+				"Engine",
+				"Slate",
+				"SlateCore",
+				"AudioMixer",
+				// ... add private dependencies that you statically link with here ...	
+			}
+			);
+		
+		
+		DynamicallyLoadedModuleNames.AddRange(
+			new string[]
+			{
+				// ... add any modules that your module loads dynamically here ...
+			}
+			);
+	}
+}
diff --git a/Engine/Plugins/DynamotorSynth/Source/DynamotorSynth/Private/DynamotorSynth.cpp b/Engine/Plugins/DynamotorSynth/Source/DynamotorSynth/Private/DynamotorSynth.cpp
new file mode 100644
--- /dev/null
+++ b/Engine/Plugins/DynamotorSynth/Source/DynamotorSynth/Private/DynamotorSynth.cpp
@@ -0,0 +1,20 @@
+// Copyright Epic Games, Inc. All Rights Reserved.
+
+#include "DynamotorSynth.h"
+
+#define LOCTEXT_NAMESPACE "FDynamotorSynthModule"
+
+void FDynamotorSynthModule::StartupModule()
+{
+	// This code will execute after your module is loaded into memory; the exact timing is specified in the .uplugin file per-module
+}
+
+void FDynamotorSynthModule::ShutdownModule()
+{
+	// This function may be called during shutdown to clean up your module.  For modules that support dynamic reloading,
+	// we call this function before unloading the module.
+}
+
+#undef LOCTEXT_NAMESPACE
+	
+IMPLEMENT_MODULE(FDynamotorSynthModule, DynamotorSynth)
\ No newline at end of file
diff --git a/Engine/Plugins/DynamotorSynth/Source/DynamotorSynth/Public/DynamotorSynth.h b/Engine/Plugins/DynamotorSynth/Source/DynamotorSynth/Public/DynamotorSynth.h
new file mode 100644
--- /dev/null
+++ b/Engine/Plugins/DynamotorSynth/Source/DynamotorSynth/Public/DynamotorSynth.h
@@ -0,0 +1,15 @@
+// Copyright Epic Games, Inc. All Rights Reserved.
+
+#pragma once
+
+#include "CoreMinimal.h"
+#include "Modules/ModuleManager.h"
+
+class FDynamotorSynthModule : public IModuleInterface
+{
+public:
+
+	/** IModuleInterface implementation */
+	virtual void StartupModule() override;
+	virtual void ShutdownModule() override;
+};
